Chaos Tests
- move Catch2 matchers and Test Scene into shared module
- add Chaos::LowLevelTest namespace for Low-Level-Test harness classes and utilities

#rnx

diff --git a/Engine/Source/Developer/ChaosTestHarness/ChaosTestHarness.Build.cs b/Engine/Source/Developer/ChaosTestHarness/ChaosTestHarness.Build.cs
new file mode 100644
--- /dev/null
+++ b/Engine/Source/Developer/ChaosTestHarness/ChaosTestHarness.Build.cs
@@ -0,0 +1,25 @@
+// Copyright Epic Games, Inc. All Rights Reserved.
+
+using UnrealBuildTool;
+
+public class ChaosTestHarness : TestModuleRules
+{
+	public ChaosTestHarness(ReadOnlyTargetRules Target) : base(Target)
+	{
+		PublicDependencyModuleNames.AddRange(new string[] {
+			"Catch2"
+		});
+
+		PrivateDependencyModuleNames.AddRange(new string[] {
+			"Core",
+			"CoreUObject",
+			"Engine",
+			"LowLevelTestsRunner",
+		});
+
+		// Temporarily link to base Chaos modules - remove once everything is moved into the new tree
+		SetupModulePhysicsSupport(Target);
+
+		bWarningsAsErrors = true;
+	}
+}
\ No newline at end of file
diff --git a/Engine/Source/Developer/ChaosTestHarness/ChaosTestHarness.Target.cs b/Engine/Source/Developer/ChaosTestHarness/ChaosTestHarness.Target.cs
new file mode 100644
--- /dev/null
+++ b/Engine/Source/Developer/ChaosTestHarness/ChaosTestHarness.Target.cs
@@ -0,0 +1,21 @@
+// Copyright Epic Games, Inc. All Rights Reserved.
+
+using UnrealBuildTool;
+
+[SupportedPlatforms(UnrealPlatformClass.All)]
+public class ChaosTestHarnessTarget : TestTargetRules
+{
+	public ChaosTestHarnessTarget(TargetInfo Target) : base(Target)
+	{
+		bCompileAgainstEngine = true;
+		bCompileAgainstCoreUObject = true;
+		bMockEngineDefaults = true;
+		bUsePlatformFileStub = true;
+
+		bWithLowLevelTestsOverride = true;
+
+		GlobalDefinitions.Add("CATCH_CONFIG_ENABLE_BENCHMARKING=1");
+
+		bWarningsAsErrors = true;
+	}
+}
\ No newline at end of file
diff --git a/Engine/Source/Developer/ChaosTestHarness/Private/ChaosTestScene.cpp b/Engine/Source/Developer/ChaosTestHarness/Private/ChaosTestScene.cpp
new file mode 100644
--- /dev/null
+++ b/Engine/Source/Developer/ChaosTestHarness/Private/ChaosTestScene.cpp
@@ -0,0 +1,252 @@
+// Copyright Epic Games, Inc. All Rights Reserved.
+#include "Chaos/LowLevelTest/ChaosTestScene.h"
+#include "Chaos/Box.h"
+#include "Chaos/MassProperties.h"
+#include "Chaos/Sphere.h"
+#include "ChaosSolversModule.h"
+#include "EventsData.h"
+#include "PBDRigidsSolver.h"
+#include "PhysicsProxy/SingleParticlePhysicsProxy.h"
+
+
+namespace Chaos::LowLevelTest
+{
+
+	FChaosTestScene::FChaosTestScene()
+		: NumTicks(0)
+	{
+		CreateSolver();
+	}
+
+	FChaosTestScene::~FChaosTestScene()
+	{
+		DestroySolver();
+	}
+
+	Chaos::FPBDRigidsSolver* FChaosTestScene::GetSolver()
+	{
+		return Solver;
+	}
+
+	void FChaosTestScene::TickSolver(float Dt)
+	{
+		Solver->AdvanceAndDispatch_External(Dt);
+		Solver->UpdateGameThreadStructures();
+		Solver->SyncEvents_GameThread();
+		++NumTicks;
+	}
+
+	bool FChaosTestScene::TickSolverToSleep(float Dt, int32 MaxTicks)
+	{
+		for (int32 SleepTickIndex = 0; (SleepTickIndex < MaxTicks) && !IsSceneSleeping(); ++SleepTickIndex)
+		{
+			TickSolver(Dt);
+		}
+
+		return IsSceneSleeping();
+	}
+
+	bool FChaosTestScene::IsSceneSleeping() const
+	{
+		using namespace Chaos;
+
+		for (Chaos::FSingleParticlePhysicsProxy* Proxy : Proxies)
+		{
+			if (Proxy->GetGameThreadAPI().ObjectState() == EObjectStateType::Dynamic)
+			{
+				return false;
+			}
+		}
+
+		return true;
+	}
+
+	int FChaosTestScene::GetNumTicks() const
+	{
+		return NumTicks;
+	}
+
+	Chaos::FSingleParticlePhysicsProxy* FChaosTestScene::AddDynamicSphere(const TCHAR* InName, const FTransform3d& InTransform, const double& InRadius)
+	{
+		using namespace Chaos;
+
+		FSingleParticlePhysicsProxy* Proxy = CreateDynamicSphere(InName, InTransform, InRadius);
+
+		Solver->RegisterObject(Proxy);
+
+		return Proxy;
+	}
+
+	Chaos::FSingleParticlePhysicsProxy* FChaosTestScene::AddDynamicBox(const TCHAR* InName, const FTransform3d& InTransform, const FVector3d& InHalfExtents)
+	{
+		using namespace Chaos;
+
+		FSingleParticlePhysicsProxy* Proxy = CreateDynamicBox(InName, InTransform, InHalfExtents);
+
+		Solver->RegisterObject(Proxy);
+
+		return Proxy;
+	}
+
+	Chaos::FSingleParticlePhysicsProxy* FChaosTestScene::AddDynamicProxy(const TCHAR* InName, const FTransform3d& InTransform, const Chaos::FImplicitObjectPtr& InGeometry)
+	{
+		using namespace Chaos;
+
+		FSingleParticlePhysicsProxy* Proxy = CreateDynamicProxy(InName, InTransform, InGeometry);
+
+		Solver->RegisterObject(Proxy);
+
+		return Proxy;
+	}
+
+	Chaos::FSingleParticlePhysicsProxy* FChaosTestScene::AddKinematicBox(const TCHAR* InName, const FTransform3d& InTransform, const FVector3d& InHalfExtents)
+	{
+		using namespace Chaos;
+
+		FSingleParticlePhysicsProxy* Proxy = CreateDynamicBox(InName, InTransform, InHalfExtents);
+		Proxy->GetGameThreadAPI().SetObjectState(EObjectStateType::Kinematic);
+
+		Solver->RegisterObject(Proxy);
+
+		return Proxy;
+	}
+
+	Chaos::FSingleParticlePhysicsProxy* FChaosTestScene::AddKinematicProxy(const TCHAR* InName, const FTransform3d& InTransform, const Chaos::FImplicitObjectPtr& InGeometry)
+	{
+		using namespace Chaos;
+
+		FSingleParticlePhysicsProxy* Proxy = CreateDynamicProxy(InName, InTransform, InGeometry);
+		Proxy->GetGameThreadAPI().SetObjectState(EObjectStateType::Kinematic);
+
+		Solver->RegisterObject(Proxy);
+
+		return Proxy;
+	}
+
+	void FChaosTestScene::SetGravity(const FVector3d& InG)
+	{
+		Solver->GetEvolution()->GetGravityForces().SetAcceleration(InG, 0);
+	}
+
+
+	void FChaosTestScene::SetCollisionCallback(const TFunction<void(const Chaos::FCollisionEventData&)>& InCollisionCallback)
+	{
+		CollisionCallback = InCollisionCallback;
+	}
+
+	void FChaosTestScene::CreateSolver()
+	{
+		using namespace Chaos;
+
+		const float AsyncDt = -1.0f;	// No async
+		Solver = FChaosSolversModule::GetModule()->CreateSolver(nullptr, AsyncDt, Chaos::EThreadingMode::SingleThread, TEXT("Collision Test Scene"));
+
+		Chaos::FEventManager* EventManager = Solver->GetEventManager();
+		EventManager->RegisterHandler<Chaos::FCollisionEventData>(Chaos::EEventType::Collision, this, &FChaosTestScene::OnCollision);
+	}
+
+	void FChaosTestScene::DestroySolver()
+	{
+		using namespace Chaos;
+
+		FChaosSolversModule::GetModule()->DestroySolver(Solver);
+		Solver = nullptr;
+	}
+
+	Chaos::FSingleParticlePhysicsProxy* FChaosTestScene::CreateDynamicSphere(const TCHAR* InName, const FTransform3d& InTransform, const double& InRadius)
+	{
+		using namespace Chaos;
+		FImplicitObjectPtr SphereGeom = FImplicitObjectPtr(new Chaos::FImplicitSphere3(FVector3f::Zero(), (float)InRadius));
+		return CreateDynamicProxy(InName, InTransform, SphereGeom);
+	}
+
+	Chaos::FSingleParticlePhysicsProxy* FChaosTestScene::CreateDynamicBox(const TCHAR* InName, const FTransform3d& InTransform, const FVector3d& InHalfExtents)
+	{
+		using namespace Chaos;
+		FVector3d HalfExtents = InHalfExtents * InTransform.GetScale3D().GetAbs();
+		FImplicitObjectPtr BoxGeom = FImplicitObjectPtr(new Chaos::FImplicitBox3(-HalfExtents, HalfExtents));
+		return CreateDynamicProxy(InName, InTransform, BoxGeom);
+	}
+
+	Chaos::FSingleParticlePhysicsProxy* FChaosTestScene::CreateDynamicProxy(const TCHAR* InName, const FTransform3d& InTransform, const Chaos::FImplicitObjectPtr& InGeometry)
+	{
+		using namespace Chaos;
+
+		FSingleParticlePhysicsProxy* Proxy = FSingleParticlePhysicsProxy::Create(Chaos::TPBDRigidParticle<FReal, 3>::CreateParticle());
+		FRigidBodyHandle_External& ParticleGT = Proxy->GetGameThreadAPI();
+
+	#if CHAOS_DEBUG_NAME
+		if (InName != nullptr)
+		{
+			TSharedPtr<FString> NamePtr = MakeShared<FString>(InName);
+			ParticleGT.SetDebugName(NamePtr);
+		}
+	#endif
+		ParticleGT.SetGeometry(InGeometry);
+		ParticleGT.SetX(InTransform.GetLocation());
+		ParticleGT.SetR(InTransform.GetRotation());
+
+		InitCollisionFilter(Proxy);
+		InitMassProperties(Proxy);
+
+		Proxies.Add(Proxy);
+
+		return Proxy;
+	}
+
+	void FChaosTestScene::InitCollisionFilter(Chaos::FSingleParticlePhysicsProxy* Proxy)
+	{
+		using namespace Chaos;
+
+		constexpr EFilterFlags FilterFlags = EFilterFlags::SimpleCollision | EFilterFlags::ContactNotify | EFilterFlags::ModifyContacts;
+		const Filter::FShapeFilterData ShapeFilter = Filter::FShapeFilterBuilder::BuildBlockAll(FilterFlags);
+
+		for (int32 ShapeIndex = 0; ShapeIndex < Proxy->GetGameThreadAPI().ShapesArray().Num(); ++ShapeIndex)
+		{
+			Proxy->GetGameThreadAPI().SetShapeSimCollisionEnabled(ShapeIndex, true);
+			Proxy->GetGameThreadAPI().SetShapeQueryCollisionEnabled(ShapeIndex, true);
+			Proxy->GetGameThreadAPI().SetShapeFilterData(ShapeIndex, ShapeFilter);
+		}
+	}
+
+	void FChaosTestScene::InitMassProperties(Chaos::FSingleParticlePhysicsProxy* Proxy)
+	{
+		using namespace Chaos;
+
+		FRigidBodyHandle_External& ParticleGT = Proxy->GetGameThreadAPI();
+		const FShapesArray& ShapesArray = ParticleGT.ShapesArray();
+
+
+		float DensityKgPerCm3 = 1.e-3f;	// Water 1g/cm3
+
+		// Calculate mass and inertia
+		FMassProperties MassProps;
+		Chaos::CalculateMassPropertiesFromShapeCollection(
+			MassProps,
+			ShapesArray.Num(),
+			DensityKgPerCm3,
+			TArray<bool>(),
+			[&ShapesArray](int32 ShapeIndex) { return ShapesArray[ShapeIndex].Get(); }
+		);
+
+		// Diagonalize the inertia
+		TransformToLocalSpace(MassProps);
+
+		// Set up particle mass
+		ParticleGT.SetM(MassProps.Mass);
+		ParticleGT.SetInvM(1.0f / MassProps.Mass);
+		ParticleGT.SetI(MassProps.InertiaTensor.GetDiagonal());
+		ParticleGT.SetInvI(FVec3(1.0f) / MassProps.InertiaTensor.GetDiagonal());
+		ParticleGT.SetCenterOfMass(MassProps.CenterOfMass);
+		ParticleGT.SetRotationOfMass(MassProps.RotationOfMass);
+	}
+
+	void FChaosTestScene::OnCollision(const Chaos::FCollisionEventData& CollisionEventData)
+	{
+		if (CollisionCallback.IsSet())
+		{
+			CollisionCallback(CollisionEventData);
+		}
+	}
+
+} // Chaos::LowLevelTest
\ No newline at end of file
diff --git a/Engine/Source/Developer/ChaosTestHarness/Private/Init.cpp b/Engine/Source/Developer/ChaosTestHarness/Private/Init.cpp
new file mode 100644
--- /dev/null
+++ b/Engine/Source/Developer/ChaosTestHarness/Private/Init.cpp
@@ -0,0 +1,35 @@
+// Copyright Epic Games, Inc. All Rights Reserved.
+
+#include "CoreMinimal.h"
+#include "Chaos/LowLevelTest/ChaosTestHarness.h"
+
+#include "Logging/LogScopedVerbosityOverride.h"
+#include "TestCommon/Initialization.h"
+
+// Pull in the log category declarations for below overrides
+#include "SlateGlobals.h"
+#include "Styling/SlateWidgetStyleContainerBase.h"
+
+#include "catch2/catch_tag_alias_autoregistrar.hpp"
+
+CATCH_REGISTER_TAG_ALIAS("[@ChaosTests]", "[Chaos]~[!benchmark]")
+CATCH_REGISTER_TAG_ALIAS("[@ChaosBenchmarks]", "[Chaos][!benchmark]")
+
+GROUP_BEFORE_GLOBAL(Catch::DefaultGroup)
+{
+	// Find the engine directory before InitAll() enables the platform file stub. Prevents warnings.
+	FPlatformMisc::EngineDir();
+
+	// Temporarily disable warnings during init
+	LOG_SCOPE_VERBOSITY_OVERRIDE(LogSlate, ELogVerbosity::Error);
+	LOG_SCOPE_VERBOSITY_OVERRIDE(LogSlateStyle, ELogVerbosity::Error);
+	LOG_SCOPE_VERBOSITY_OVERRIDE(LogUObjectGlobals, ELogVerbosity::Error);
+
+	InitAll(true, true);
+}
+
+GROUP_AFTER_GLOBAL(Catch::DefaultGroup)
+{
+	CleanupAll();
+	GIsRunning = false;
+}
\ No newline at end of file
diff --git a/Engine/Source/Developer/ChaosTestHarness/Public/Chaos/LowLevelTest/ChaosTestCatch2.h b/Engine/Source/Developer/ChaosTestHarness/Public/Chaos/LowLevelTest/ChaosTestCatch2.h
new file mode 100644
--- /dev/null
+++ b/Engine/Source/Developer/ChaosTestHarness/Public/Chaos/LowLevelTest/ChaosTestCatch2.h
@@ -0,0 +1,32 @@
+// Copyright Epic Games, Inc. All Rights Reserved.
+
+#pragma once
+
+// The base test harness file doesn't include the benchmark files from Catch2
+// so those macros show up as undefined, so we add those files here.
+
+#if defined(THIRD_PARTY_INCLUDES_START) && defined(THIRD_PARTY_INCLUDES_END)
+THIRD_PARTY_INCLUDES_START
+#endif // defined(THIRD_PARTY_INCLUDES_START) && defined(THIRD_PARTY_INCLUDES_END)
+
+#ifdef _MSC_VER
+#pragma pack(push, 8)
+#pragma warning(push)
+#pragma warning(disable: 4005) // 'identifier': macro redefinition
+#pragma warning(disable: 4582) // 'type': constructor is not implicitly called
+#pragma warning(disable: 4583) // 'type': destructor is not implicitly called
+#endif
+#include <catch2/benchmark/catch_benchmark_all.hpp>
+#include <catch2/generators/catch_generators.hpp>
+#ifdef _MSC_VER
+#pragma warning(pop)
+#pragma pack(pop)
+#endif
+
+#if defined(THIRD_PARTY_INCLUDES_START) && defined(THIRD_PARTY_INCLUDES_END)
+THIRD_PARTY_INCLUDES_END
+#endif // defined(THIRD_PARTY_INCLUDES_START) && defined(THIRD_PARTY_INCLUDES_END)
+
+// Support for e.g.:
+// CHECK_THAT(Z, Catch::Matchers::WithinAbs(100.0f, 0.1f));
+#include <catch2/matchers/catch_matchers_floating_point.hpp>
diff --git a/Engine/Source/Developer/ChaosTestHarness/Public/Chaos/LowLevelTest/ChaosTestHarness.h b/Engine/Source/Developer/ChaosTestHarness/Public/Chaos/LowLevelTest/ChaosTestHarness.h
new file mode 100644
--- /dev/null
+++ b/Engine/Source/Developer/ChaosTestHarness/Public/Chaos/LowLevelTest/ChaosTestHarness.h
@@ -0,0 +1,10 @@
+// Copyright Epic Games, Inc. All Rights Reserved.
+
+#pragma once
+
+// From LowLevelTests Module
+#include "TestHarness.h"
+
+#include "Chaos/LowLevelTest/ChaosTestCatch2.h"
+#include "Chaos/LowLevelTest/ChaosTestStringMakers.h"
+#include "Chaos/LowLevelTest/ChaosTestMatchers.h"
diff --git a/Engine/Source/Developer/ChaosTestHarness/Public/Chaos/LowLevelTest/ChaosTestMatchers.h b/Engine/Source/Developer/ChaosTestHarness/Public/Chaos/LowLevelTest/ChaosTestMatchers.h
new file mode 100644
--- /dev/null
+++ b/Engine/Source/Developer/ChaosTestHarness/Public/Chaos/LowLevelTest/ChaosTestMatchers.h
@@ -0,0 +1,47 @@
+// Copyright Epic Games, Inc. All Rights Reserved.
+#pragma once
+
+#include "CoreMinimal.h"
+#include "Chaos/Core.h"
+#include "Chaos/LowLevelTest/ChaosTestCatch2.h"
+
+namespace Catch
+{
+	template <typename T, int D>
+	class FVectorEqualityMatcher : public Catch::Matchers::MatcherBase<Chaos::TVector<T, D>>
+	{
+	public:
+		using FVector = Chaos::TVector<T, D>;
+		FVector Target;
+		T Epsilon;
+
+		FVectorEqualityMatcher(const FVector& InTarget, const T InEpsilon = UE_KINDA_SMALL_NUMBER)
+			: Target(InTarget)
+			, Epsilon(InEpsilon)
+		{
+		}
+
+		// Performs the test for this matcher
+		bool match(const FVector& ToCompare) const override
+		{
+			return FVector::PointsAreNear(ToCompare, Target, Epsilon);
+		}
+
+		// Produces a string describing what this matcher does. It should
+		// include any provided data (the begin/ end in this case) and
+		// be written as if it were stating a fact (in the output it will be
+		// preceded by the value under test).
+		virtual std::string describe() const override
+		{
+			std::ostringstream Stream;
+			Stream << "is equal to (" << Catch::StringMaker<FVector>::convert(Target) << ") with epsilon " << Epsilon;
+			return Stream.str();
+		}
+	};
+
+	template <typename T, int D>
+	FVectorEqualityMatcher<T, D> ApproxEq(const Chaos::TVector<T, D>& Target, const T Epsilon = UE_KINDA_SMALL_NUMBER)
+	{
+		return FVectorEqualityMatcher<T, D>(Target, Epsilon);
+	}
+} // namespace Catch 
\ No newline at end of file
diff --git a/Engine/Source/Developer/ChaosTestHarness/Public/Chaos/LowLevelTest/ChaosTestScene.h b/Engine/Source/Developer/ChaosTestHarness/Public/Chaos/LowLevelTest/ChaosTestScene.h
new file mode 100644
--- /dev/null
+++ b/Engine/Source/Developer/ChaosTestHarness/Public/Chaos/LowLevelTest/ChaosTestScene.h
@@ -0,0 +1,93 @@
+// Copyright Epic Games, Inc. All Rights Reserved.
+#pragma once
+
+#include "CoreMinimal.h"
+#include "Chaos/ImplicitFwd.h"
+#include "Chaos/SimCallbackObject.h"
+
+namespace Chaos
+{
+	class FPBDRigidsSolver;
+	class FSingleParticlePhysicsProxy;
+
+	struct FCollisionEventData;
+}
+
+namespace Chaos::LowLevelTest
+{
+	//
+	// A wrapper around a chaos solver providing minimal functionality for testing.
+	//
+	class UE_INTERNAL FChaosTestScene
+	{
+	public:
+		UE_INTERNAL CHAOSTESTHARNESS_API FChaosTestScene();
+		UE_INTERNAL CHAOSTESTHARNESS_API ~FChaosTestScene();
+
+		UE_INTERNAL CHAOSTESTHARNESS_API Chaos::FPBDRigidsSolver* GetSolver();
+
+		// Tick the solver.
+		UE_INTERNAL CHAOSTESTHARNESS_API void TickSolver(float Dt);
+
+		// Tick the solver until it sleeps (or until we reach MaxTicks ticks).
+		// Return true if the scene is sleeping.
+		UE_INTERNAL CHAOSTESTHARNESS_API bool TickSolverToSleep(float Dt, int32 MaxTicks);
+
+		// Return true if all dynamic bodies are asleep.
+		UE_INTERNAL CHAOSTESTHARNESS_API bool IsSceneSleeping() const;
+
+		// How many ticks have been run (useful if using TickSolverToSleep)
+		UE_INTERNAL CHAOSTESTHARNESS_API int GetNumTicks() const;
+
+		// Add a dynamic awake body to the scene.
+		UE_INTERNAL CHAOSTESTHARNESS_API Chaos::FSingleParticlePhysicsProxy* AddDynamicSphere(const TCHAR* InName, const FTransform3d& InTransform, const double& InRadius);
+		UE_INTERNAL CHAOSTESTHARNESS_API Chaos::FSingleParticlePhysicsProxy* AddDynamicBox(const TCHAR* InName, const FTransform3d& InTransform, const FVector3d& InHalfExtents);
+		UE_INTERNAL CHAOSTESTHARNESS_API Chaos::FSingleParticlePhysicsProxy* AddDynamicProxy(const TCHAR* InName, const FTransform3d& InTransform, const Chaos::FImplicitObjectPtr& InGeometry);
+
+		// Add a kinematic (which can be promoted to dynamic) body to the scene
+		UE_INTERNAL CHAOSTESTHARNESS_API Chaos::FSingleParticlePhysicsProxy* AddKinematicBox(const TCHAR* InName, const FTransform3d& InTransform, const FVector3d& InHalfExtents);
+		UE_INTERNAL CHAOSTESTHARNESS_API Chaos::FSingleParticlePhysicsProxy* AddKinematicProxy(const TCHAR* InName, const FTransform3d& InTransform, const Chaos::FImplicitObjectPtr& InGeometry);
+
+		// Set the gravity (for gravity group 0)
+		UE_INTERNAL CHAOSTESTHARNESS_API void SetGravity(const FVector3d& InG);
+
+		// Set the collision event callback (can only be one)
+		UE_INTERNAL CHAOSTESTHARNESS_API void SetCollisionCallback(const TFunction<void(const Chaos::FCollisionEventData&)>& InCollisionCallback);
+
+	private:
+
+		void CreateSolver();
+		void DestroySolver();
+
+		Chaos::FSingleParticlePhysicsProxy* CreateDynamicSphere(const TCHAR* InName, const FTransform3d& InTransform, const double& InRadius);
+		Chaos::FSingleParticlePhysicsProxy* CreateDynamicBox(const TCHAR* InName, const FTransform3d& InTransform, const FVector3d& InHalfExtents);
+		Chaos::FSingleParticlePhysicsProxy* CreateDynamicProxy(const TCHAR* InName, const FTransform3d& InTransform, const Chaos::FImplicitObjectPtr& InGeometry);
+		void InitCollisionFilter(Chaos::FSingleParticlePhysicsProxy* Proxy);
+		void InitMassProperties(Chaos::FSingleParticlePhysicsProxy* Proxy);
+
+		void OnCollision(const Chaos::FCollisionEventData& CollisionEventData);
+
+	private:
+		Chaos::FPBDRigidsSolver* Solver;
+		TArray<Chaos::FSingleParticlePhysicsProxy*> Proxies;
+		TFunction<void(const Chaos::FCollisionEventData&)> CollisionCallback;
+		int32 NumTicks;
+	};
+
+	// A sim callback that allows unit tests to hook a lambda up to the various callbacks
+	// TODO: Add all the other callbacks!
+	class FChaosTestSimCallback : public Chaos::TSimCallbackObject<Chaos::FSimCallbackNoInput, Chaos::FSimCallbackNoOutput, Chaos::ESimCallbackOptions::ContactModification>
+	{
+	public:
+		virtual void OnContactModification_Internal(Chaos::FCollisionContactModifier& Modifier) override final
+		{
+			if (ContactModifier)
+			{
+				ContactModifier(Modifier);
+			}
+		}
+
+		TFunction<void(Chaos::FCollisionContactModifier&)> ContactModifier;
+	};
+
+}
\ No newline at end of file
diff --git a/Engine/Source/Developer/ChaosTestHarness/Public/Chaos/LowLevelTest/ChaosTestStringMakers.h b/Engine/Source/Developer/ChaosTestHarness/Public/Chaos/LowLevelTest/ChaosTestStringMakers.h
new file mode 100644
--- /dev/null
+++ b/Engine/Source/Developer/ChaosTestHarness/Public/Chaos/LowLevelTest/ChaosTestStringMakers.h
@@ -0,0 +1,47 @@
+// Copyright Epic Games, Inc. All Rights Reserved.
+#pragma once
+
+#include "Chaos/Core.h"
+#include "Chaos/ChaosLogUtil.h"
+#include "Chaos/LowLevelTest/ChaosTestCatch2.h"
+
+#include <string>
+
+namespace Catch 
+{
+	template<>
+	struct StringMaker<Chaos::FRotation3>
+	{
+		static std::string convert(Chaos::FRotation3 const& Value) { return std::string(TCHAR_TO_UTF8(*Chaos::ToString(Value))); }
+	};
+
+	template<>
+	struct StringMaker<Chaos::FRotation3f>
+	{
+		static std::string convert(Chaos::FRotation3f const& Value) { return std::string(TCHAR_TO_UTF8(*Chaos::ToString(Value))); }
+	};
+
+	template<>
+	struct StringMaker<Chaos::FVec3>
+	{
+		static std::string convert(Chaos::FVec3 const& Value) { return std::string(TCHAR_TO_UTF8(*Chaos::ToString(Value))); }
+	};
+
+	template<>
+	struct StringMaker<Chaos::FVec3f>
+	{
+		static std::string convert(Chaos::FVec3f const& Value) { return std::string(TCHAR_TO_UTF8(*Chaos::ToString(Value))); }
+	};
+
+	template<>
+	struct StringMaker<FVector3d>
+	{
+		static std::string convert(FVector3d const& Value) { return StringMaker<Chaos::FVec3>::convert(Chaos::FVec3(Value)); }
+	};
+
+	template<>
+	struct StringMaker<FVector3f>
+	{
+		static std::string convert(FVector3f const& Value) { return StringMaker<Chaos::FVec3f>::convert(Chaos::FVec3f(Value)); }
+	};
+}
\ No newline at end of file
