Chaos Test - protect against zero mass in FChaosTestScene::InitMassProperties

#rnx

diff --git a/Engine/Source/Developer/ChaosTestHarness/Private/ChaosTestScene.cpp b/Engine/Source/Developer/ChaosTestHarness/Private/ChaosTestScene.cpp
--- a/Engine/Source/Developer/ChaosTestHarness/Private/ChaosTestScene.cpp
+++ b/Engine/Source/Developer/ChaosTestHarness/Private/ChaosTestScene.cpp
@@ -233,10 +233,26 @@ namespace Chaos::LowLevelTest
 		TransformToLocalSpace(MassProps);
 
 		// Set up particle mass
-		ParticleGT.SetM(MassProps.Mass);
-		ParticleGT.SetInvM(1.0f / MassProps.Mass);
-		ParticleGT.SetI(MassProps.InertiaTensor.GetDiagonal());
-		ParticleGT.SetInvI(FVec3(1.0f) / MassProps.InertiaTensor.GetDiagonal());
+		if (MassProps.Mass > 0)
+		{
+			ParticleGT.SetM(MassProps.Mass);
+			ParticleGT.SetInvM(1.0f / MassProps.Mass);
+		}
+		else
+		{
+			ParticleGT.SetM(0);
+			ParticleGT.SetInvM(0);
+		}
+		if (MassProps.InertiaTensor.GetDiagonal().GetAbsMin() > 0)
+		{
+			ParticleGT.SetI(MassProps.InertiaTensor.GetDiagonal());
+			ParticleGT.SetInvI(FVec3(1.0f) / MassProps.InertiaTensor.GetDiagonal());
+		}
+		else
+		{
+			ParticleGT.SetI(FVec3(0));
+			ParticleGT.SetInvI(FVec3(0));
+		}
 		ParticleGT.SetCenterOfMass(MassProps.CenterOfMass);
 		ParticleGT.SetRotationOfMass(MassProps.RotationOfMass);
 	}
diff --git a/Engine/Source/Developer/ChaosTestHarness/Public/Chaos/LowLevelTest/ChaosTestHarness.h b/Engine/Source/Developer/ChaosTestHarness/Public/Chaos/LowLevelTest/ChaosTestHarness.h
--- a/Engine/Source/Developer/ChaosTestHarness/Public/Chaos/LowLevelTest/ChaosTestHarness.h
+++ b/Engine/Source/Developer/ChaosTestHarness/Public/Chaos/LowLevelTest/ChaosTestHarness.h
@@ -2,7 +2,7 @@
 
 #pragma once
 
-// From LowLevelTests Module
+// From LowLevelTest Module
 #include "TestHarness.h"
 
 #include "Chaos/LowLevelTest/ChaosTestCatch2.h"
