Adding a test matcher for AABB equality

#Synced-CL 48505955

diff --git a/Engine/Source/Developer/ChaosTestHarness/Public/Chaos/LowLevelTest/ChaosTestMatchers.h b/Engine/Source/Developer/ChaosTestHarness/Public/Chaos/LowLevelTest/ChaosTestMatchers.h
--- a/Engine/Source/Developer/ChaosTestHarness/Public/Chaos/LowLevelTest/ChaosTestMatchers.h
+++ b/Engine/Source/Developer/ChaosTestHarness/Public/Chaos/LowLevelTest/ChaosTestMatchers.h
@@ -43,4 +43,52 @@ namespace Catch
 	{
 		return FVectorEqualityMatcher<T, D>(Target, Epsilon);
 	}
+	
+	template <typename T, int D>
+	class TAABBEqualityMatcher : public Catch::Matchers::MatcherBase<Chaos::TAABB<T, D>>
+	{
+	public:
+		using FAABB = Chaos::TAABB<T, D>;
+		FAABB Target;
+		T Epsilon;
+
+		TAABBEqualityMatcher(const FAABB& InTarget, const T InEpsilon = UE_KINDA_SMALL_NUMBER)
+			: Target(InTarget)
+			, Epsilon(InEpsilon)
+		{
+		}
+
+		// Performs the test for this matcher
+		bool match(const FAABB& ToCompare) const override
+		{
+			if (Epsilon == 0)
+			{
+				return ToCompare.Min() == Target.Min() && ToCompare.Max() == Target.Max();
+			}
+			return FVector::PointsAreNear(ToCompare.Min(), Target.Min(), Epsilon) && FVector::PointsAreNear(ToCompare.Max(), Target.Max(), Epsilon);
+		}
+
+		// Produces a string describing what this matcher does. It should
+		// include any provided data (the begin/ end in this case) and
+		// be written as if it were stating a fact (in the output it will be
+		// preceded by the value under test).
+		virtual std::string describe() const override
+		{
+			std::ostringstream Stream;
+			Stream << "is equal to (" << Catch::StringMaker<FAABB>::convert(Target) << ") with epsilon " << Epsilon;
+			return Stream.str();
+		}
+	};
+	
+	template <typename T, int D>
+	TAABBEqualityMatcher<T, D> Equal(const Chaos::TAABB<T, D>& Target)
+	{
+		return TAABBEqualityMatcher<T, D>(Target, 0);
+	}
+	
+	template <typename T, int D>
+	TAABBEqualityMatcher<T, D> ApproxyEq(const Chaos::TAABB<T, D>& Target, const T Epsilon = UE_KINDA_SMALL_NUMBER)
+	{
+		return TAABBEqualityMatcher<T, D>(Target, Epsilon);
+	}
 } // namespace Catch 
\ No newline at end of file
diff --git a/Engine/Source/Developer/ChaosTestHarness/Public/Chaos/LowLevelTest/ChaosTestStringMakers.h b/Engine/Source/Developer/ChaosTestHarness/Public/Chaos/LowLevelTest/ChaosTestStringMakers.h
--- a/Engine/Source/Developer/ChaosTestHarness/Public/Chaos/LowLevelTest/ChaosTestStringMakers.h
+++ b/Engine/Source/Developer/ChaosTestHarness/Public/Chaos/LowLevelTest/ChaosTestStringMakers.h
@@ -2,6 +2,7 @@
 #pragma once
 
 #include "Chaos/Core.h"
+#include "Chaos/AABB.h"
 #include "Chaos/ChaosLogUtil.h"
 #include "Chaos/LowLevelTest/ChaosTestCatch2.h"
 
@@ -44,4 +45,10 @@ namespace Catch
 	{
 		static std::string convert(FVector3f const& Value) { return StringMaker<Chaos::FVec3f>::convert(Chaos::FVec3f(Value)); }
 	};
+
+	template<typename T, int d>
+	struct StringMaker<Chaos::TAABB<T, d>>
+	{
+		static std::string convert(Chaos::TAABB<T, d> const& Value) { return StringMaker<FString>::convert(Value.ToString()); }
+	};
 }
\ No newline at end of file
