Fixing minor issues in chaos test matchers

#Synced-CL 48505955

diff --git a/Engine/Source/Developer/ChaosTestHarness/Public/Chaos/LowLevelTest/ChaosTestMatchers.h b/Engine/Source/Developer/ChaosTestHarness/Public/Chaos/LowLevelTest/ChaosTestMatchers.h
--- a/Engine/Source/Developer/ChaosTestHarness/Public/Chaos/LowLevelTest/ChaosTestMatchers.h
+++ b/Engine/Source/Developer/ChaosTestHarness/Public/Chaos/LowLevelTest/ChaosTestMatchers.h
@@ -52,7 +52,7 @@ namespace Catch
 		FAABB Target;
 		T Epsilon;
 
-		TAABBEqualityMatcher(const FAABB& InTarget, const T InEpsilon = UE_KINDA_SMALL_NUMBER)
+		TAABBEqualityMatcher(const FAABB& InTarget, const T InEpsilon = static_cast<T>(UE_KINDA_SMALL_NUMBER))
 			: Target(InTarget)
 			, Epsilon(InEpsilon)
 		{
@@ -87,7 +87,7 @@ namespace Catch
 	}
 	
 	template <typename T, int D>
-	TAABBEqualityMatcher<T, D> ApproxyEq(const Chaos::TAABB<T, D>& Target, const T Epsilon = UE_KINDA_SMALL_NUMBER)
+	TAABBEqualityMatcher<T, D> ApproxEq(const Chaos::TAABB<T, D>& Target, const T Epsilon = static_cast<T>(UE_KINDA_SMALL_NUMBER))
 	{
 		return TAABBEqualityMatcher<T, D>(Target, Epsilon);
 	}
