Did some basic fixes for ChaosRigidPhysicsAsync.
- Fixed the Test config not compiling
- Fixed misc small compilation issues
- Fixed the perf comparison tests
The perf comparison tests were not setting a scene for the proxy comparisons which meant the spatial partition updates were skipped. With this fixed, the body count was lowered so tests take 5s instead of 1m.

#Synced-CL 50350999

diff --git a/Engine/Plugins/Experimental/ChaosRigidPhysicsAsync/Source/ChaosRigidPhysicsAsync/Private/ChaosRigidPhysicsAsync/RigidBodyAsync.cpp b/Engine/Plugins/Experimental/ChaosRigidPhysicsAsync/Source/ChaosRigidPhysicsAsync/Private/ChaosRigidPhysicsAsync/RigidBodyAsync.cpp
--- a/Engine/Plugins/Experimental/ChaosRigidPhysicsAsync/Source/ChaosRigidPhysicsAsync/Private/ChaosRigidPhysicsAsync/RigidBodyAsync.cpp
+++ b/Engine/Plugins/Experimental/ChaosRigidPhysicsAsync/Source/ChaosRigidPhysicsAsync/Private/ChaosRigidPhysicsAsync/RigidBodyAsync.cpp
@@ -404,7 +404,7 @@ namespace Chaos::Rigids::Async
 			UE_LOGFMT(
 				LogRigidPhysics, Error, 
 				"InitTransform and UpdateTransform on Static or Kinematic Body '{0}' not supported on the Physics Thread. Transform would diverge from Game Thread.", 
-				GetProxyAPI()->DebugName().IsValid()? **GetProxyAPI()->DebugName() : TEXT("Unknown")
+				*GetDebugName()
 				);
 			return;
 		}
@@ -460,6 +460,17 @@ namespace Chaos::Rigids::Async
 		//FGenericParticleHandle ParticleHandle = Proxy->GetHandle_LowLevel();
 		//return ParticleHandle->SetM(InMass);
 	}
+	
+	const FString FRigidBodyAsyncPT::GetDebugName() const
+	{
+		static FString UnknownName(TEXT("Unknown"));
+#if CHAOS_DEBUG_NAME
+		const FRigidBodyHandle_Internal* ProxyAPI = GetProxyAPI();
+		return ProxyAPI->DebugName().IsValid() ? *ProxyAPI->DebugName() : UnknownName;
+#else
+		return UnknownName;
+#endif
+	};
 }
 
 #endif // UE_RIGIDPHYSICS_API_ENABLED
\ No newline at end of file
diff --git a/Engine/Plugins/Experimental/ChaosRigidPhysicsAsync/Source/ChaosRigidPhysicsAsync/Public/ChaosRigidPhysicsAsync/RigidBodyAsync.h b/Engine/Plugins/Experimental/ChaosRigidPhysicsAsync/Source/ChaosRigidPhysicsAsync/Public/ChaosRigidPhysicsAsync/RigidBodyAsync.h
--- a/Engine/Plugins/Experimental/ChaosRigidPhysicsAsync/Source/ChaosRigidPhysicsAsync/Public/ChaosRigidPhysicsAsync/RigidBodyAsync.h
+++ b/Engine/Plugins/Experimental/ChaosRigidPhysicsAsync/Source/ChaosRigidPhysicsAsync/Public/ChaosRigidPhysicsAsync/RigidBodyAsync.h
@@ -143,6 +143,7 @@ namespace Chaos::Rigids::Async
 
 	private:
 		UE_INTERNAL void SetTransform(const FTransform& InTransform, bool bIsTeleport);
+		const FString GetDebugName() const;
 
 		UE::Physics::FRigidBodyHandle Handle;
 		FSingleParticlePhysicsProxy* Proxy;
diff --git a/Engine/Plugins/Experimental/ChaosRigidPhysicsAsync/Tests/ChaosRigidPhysicsAsyncTests.Target.cs b/Engine/Plugins/Experimental/ChaosRigidPhysicsAsync/Tests/ChaosRigidPhysicsAsyncTests.Target.cs
--- a/Engine/Plugins/Experimental/ChaosRigidPhysicsAsync/Tests/ChaosRigidPhysicsAsyncTests.Target.cs
+++ b/Engine/Plugins/Experimental/ChaosRigidPhysicsAsync/Tests/ChaosRigidPhysicsAsyncTests.Target.cs
@@ -7,6 +7,9 @@ using UnrealBuildTool;
  {
      public ChaosRigidPhysicsAsyncTestsTarget(TargetInfo Target) : base(Target)
      {
+		bool bIsShipping = Target.Configuration == UnrealTargetConfiguration.Shipping;
+		bool bIsTest = Target.Configuration == UnrealTargetConfiguration.Test;
+
 		bTestsRequireEngine = true;
 		bTestsRequireCoreUObject = true;
 
@@ -21,7 +24,10 @@ using UnrealBuildTool;
 		GlobalDefinitions.Add("UE_RIGIDPHYSICS_API_ENABLED=1");
 
 		// Enable constraint handle debugging in Chaos
-		GlobalDefinitions.Add("CHAOS_CONSTRAINTHANDLE_DEBUG_ENABLED=1");
+		if (!(bIsShipping || bIsTest))
+		{
+			GlobalDefinitions.Add("CHAOS_CONSTRAINTHANDLE_DEBUG_ENABLED=1");
+		}
 
 		GlobalDefinitions.Add("CATCH_CONFIG_ENABLE_BENCHMARKING=1");
 
diff --git a/Engine/Plugins/Experimental/ChaosRigidPhysicsAsync/Tests/Private/SceneLockBenchmarks.cpp b/Engine/Plugins/Experimental/ChaosRigidPhysicsAsync/Tests/Private/SceneLockBenchmarks.cpp
--- a/Engine/Plugins/Experimental/ChaosRigidPhysicsAsync/Tests/Private/SceneLockBenchmarks.cpp
+++ b/Engine/Plugins/Experimental/ChaosRigidPhysicsAsync/Tests/Private/SceneLockBenchmarks.cpp
@@ -129,7 +129,7 @@ namespace Chaos::LowLevelTest
 		UE_LOG(LogRigidPhysics, Error, TEXT("Benchmarks with asserts enabled will not give best results"));
 #endif
 
-		int32 NumBodies = 10000;
+		int32 NumBodies = 3000;
 		FTransform Transform = FTransform::Identity;
 
 		FRigidFactoryAsync Factory;
@@ -204,7 +204,7 @@ namespace Chaos::LowLevelTest
 		using namespace Chaos;
 		using namespace UE::Physics;
 
-		int32 NumBodies = 10000;
+		int32 NumBodies = 3000;
 		FTransform Transform = FTransform::Identity;
 
 		FChaosTestScene Scene;
diff --git a/Engine/Source/Developer/ChaosTestHarness/Private/ChaosTestScene.cpp b/Engine/Source/Developer/ChaosTestHarness/Private/ChaosTestScene.cpp
--- a/Engine/Source/Developer/ChaosTestHarness/Private/ChaosTestScene.cpp
+++ b/Engine/Source/Developer/ChaosTestHarness/Private/ChaosTestScene.cpp
@@ -7,6 +7,7 @@
 #include "EventsData.h"
 #include "PBDRigidsSolver.h"
 #include "PhysicsProxy/SingleParticlePhysicsProxy.h"
+#include "Chaos/ChaosScene.h"
 
 
 namespace Chaos::LowLevelTest
@@ -15,12 +16,14 @@ namespace Chaos::LowLevelTest
 	FChaosTestScene::FChaosTestScene()
 		: NumTicks(0)
 	{
+		CreateScene();
 		CreateSolver();
 	}
 
 	FChaosTestScene::~FChaosTestScene()
 	{
 		DestroySolver();
+		DestroyScene();
 	}
 
 	Chaos::FPBDRigidsSolver* FChaosTestScene::GetSolver()
@@ -134,12 +137,25 @@ namespace Chaos::LowLevelTest
 		CollisionCallback = InCollisionCallback;
 	}
 
+	void FChaosTestScene::CreateScene()
+	{
+		float AsyncDt = -1.0f;
+		Scene = new FChaosScene(nullptr, AsyncDt);
+	}
+
+	void FChaosTestScene::DestroyScene()
+	{
+		delete Scene;
+		Scene = nullptr;
+	}
+
 	void FChaosTestScene::CreateSolver()
 	{
 		using namespace Chaos;
 
 		const float AsyncDt = -1.0f;	// No async
 		Solver = FChaosSolversModule::GetModule()->CreateSolver(nullptr, AsyncDt, Chaos::EThreadingMode::SingleThread, TEXT("Collision Test Scene"));
+		Solver->PhysSceneHack = Scene;
 
 		Chaos::FEventManager* EventManager = Solver->GetEventManager();
 		EventManager->RegisterHandler<Chaos::FCollisionEventData>(Chaos::EEventType::Collision, this, &FChaosTestScene::OnCollision);
diff --git a/Engine/Source/Developer/ChaosTestHarness/Public/Chaos/LowLevelTest/ChaosTestScene.h b/Engine/Source/Developer/ChaosTestHarness/Public/Chaos/LowLevelTest/ChaosTestScene.h
--- a/Engine/Source/Developer/ChaosTestHarness/Public/Chaos/LowLevelTest/ChaosTestScene.h
+++ b/Engine/Source/Developer/ChaosTestHarness/Public/Chaos/LowLevelTest/ChaosTestScene.h
@@ -5,6 +5,8 @@
 #include "Chaos/ImplicitFwd.h"
 #include "Chaos/SimCallbackObject.h"
 
+class FChaosScene;
+
 namespace Chaos
 {
 	class FPBDRigidsSolver;
@@ -55,7 +57,8 @@ namespace Chaos::LowLevelTest
 		UE_INTERNAL CHAOSTESTHARNESS_API void SetCollisionCallback(const TFunction<void(const Chaos::FCollisionEventData&)>& InCollisionCallback);
 
 	private:
-
+		void CreateScene();
+		void DestroyScene();
 		void CreateSolver();
 		void DestroySolver();
 
@@ -68,6 +71,7 @@ namespace Chaos::LowLevelTest
 		void OnCollision(const Chaos::FCollisionEventData& CollisionEventData);
 
 	private:
+		FChaosScene* Scene = nullptr;
 		Chaos::FPBDRigidsSolver* Solver;
 		TArray<Chaos::FSingleParticlePhysicsProxy*> Proxies;
 		TFunction<void(const Chaos::FCollisionEventData&)> CollisionCallback;
diff --git a/Engine/Source/Runtime/Experimental/RigidPhysics/Public/RigidPhysics/RigidBody.h b/Engine/Source/Runtime/Experimental/RigidPhysics/Public/RigidPhysics/RigidBody.h
--- a/Engine/Source/Runtime/Experimental/RigidPhysics/Public/RigidPhysics/RigidBody.h
+++ b/Engine/Source/Runtime/Experimental/RigidPhysics/Public/RigidPhysics/RigidBody.h
@@ -530,7 +530,7 @@ namespace UE::Physics
 	template<typename ContextType>
 	TRigidBodyPtr<ContextType> TRigidBodyContainerPtrImpl<ContextType>::CreateBody(const FRigidDebugName& InName, ERigidMovementType InMovementType) const
 	{
-		return BodyContainerRaw >CreateBody(InName, InMovementType);
+		return BodyContainerRaw->CreateBody(InName, InMovementType);
 	}
 
 	template<typename ContextType>
