Fix for small compilation issue with ChaosPhysicsTests on console.

#Synced-CL 52060956

diff --git a/Engine/Source/Developer/ChaosTestHarness/Public/Chaos/LowLevelTest/ChaosTestMatchers.h b/Engine/Source/Developer/ChaosTestHarness/Public/Chaos/LowLevelTest/ChaosTestMatchers.h
--- a/Engine/Source/Developer/ChaosTestHarness/Public/Chaos/LowLevelTest/ChaosTestMatchers.h
+++ b/Engine/Source/Developer/ChaosTestHarness/Public/Chaos/LowLevelTest/ChaosTestMatchers.h
@@ -3,6 +3,7 @@
 
 #include "Chaos/Core.h"
 #include "Chaos/LowLevelTest/ChaosTestCatch2.h"
+#include <sstream>
 
 namespace Catch
 {
